Allow precompiling for Android

diff --git a/Engine/Plugins/Runtime/Oculus/OculusAvatar/Source/OculusAvatar.build.cs b/Engine/Plugins/Runtime/Oculus/OculusAvatar/Source/OculusAvatar.build.cs
--- a/Engine/Plugins/Runtime/Oculus/OculusAvatar/Source/OculusAvatar.build.cs
+++ b/Engine/Plugins/Runtime/Oculus/OculusAvatar/Source/OculusAvatar.build.cs
@@ -56,9 +56,9 @@ public class OculusAvatar : ModuleRules
             PublicDelayLoadDLLs.Add("OVRPlugin.dll");
             RuntimeDependencies.Add("$(EngineDir)/Binaries/ThirdParty/Oculus/OVRPlugin/OVRPlugin/" + Target.Platform.ToString() + "/OVRPlugin.dll");
         }
-        else
+        else if (Target.Platform != UnrealTargetPlatform.Android)
         {
-            PrecompileForTargets = PrecompileTargetsType.None;
-        }
+			PrecompileForTargets = PrecompileTargetsType.None;
+		}
     }
 }
