Fix for missing DLL when Oculus Avatar plugin is enabled

diff --git a/Engine/Plugins/Runtime/Oculus/OculusAvatar/Source/Public/OvrAvatarManager.cpp b/Engine/Plugins/Runtime/Oculus/OculusAvatar/Source/Public/OvrAvatarManager.cpp
--- a/Engine/Plugins/Runtime/Oculus/OculusAvatar/Source/Public/OvrAvatarManager.cpp
+++ b/Engine/Plugins/Runtime/Oculus/OculusAvatar/Source/Public/OvrAvatarManager.cpp
@@ -122,6 +122,12 @@ FOvrAvatarManager::~FOvrAvatarManager()
 		FPlatformProcess::FreeDllHandle(OVRPluginHandle);
 		OVRPluginHandle = nullptr;
 	}
+
+	if (OVRAvatarHandle)
+	{
+		FPlatformProcess::FreeDllHandle(OVRAvatarHandle);
+		OVRAvatarHandle = nullptr;
+	}
 }
 
 static const FString sTextureFormatStrings[ovrAvatarTextureFormat_Count] =
@@ -191,18 +197,26 @@ void FOvrAvatarManager::InitializeSDK()
 			OVRPluginHandle = FOculusHMDModule::GetOVRPluginHandle();
 		}
 
-		IsInitialized = true;
 
 #if PLATFORM_ANDROID
 		AVATAR_APP_ID = TCHAR_TO_ANSI(*GConfig->GetStr(TEXT("OnlineSubsystemOculus"), TEXT("GearVRAppId"), GEngineIni));
 		UE_LOG(LogAvatars, Display, TEXT("ovrAvatar_InitializeAndroid"));
 		ovrAvatar_InitializeAndroid(AVATAR_APP_ID, FAndroidApplication::GetGameActivityThis(), FAndroidApplication::GetJavaEnv());
 #else
+		OVRAvatarHandle = FPlatformProcess::GetDllHandle(TEXT("libovravatar.dll"));
+		if (OVRAvatarHandle == nullptr)
+		{
+			UE_LOG(LogAvatars, Log, TEXT("OVRAvatar DLL not found!"));
+			return;
+		}
+
 		AVATAR_APP_ID = TCHAR_TO_ANSI(*GConfig->GetStr(TEXT("OnlineSubsystemOculus"), TEXT("RiftAppId"), GEngineIni));
 		UE_LOG(LogAvatars, Display, TEXT("ovrAvatar_Initialize"));
 		ovrAvatar_Initialize(AVATAR_APP_ID);
 #endif
 
+		IsInitialized = true;
+
 		ovrAvatar_SetLoggingLevel(LogLevel);
 
 		// Clear avatar message queue in case there are leftover/invalid messages from other sessions/apps
diff --git a/Engine/Plugins/Runtime/Oculus/OculusAvatar/Source/Public/OvrAvatarManager.h b/Engine/Plugins/Runtime/Oculus/OculusAvatar/Source/Public/OvrAvatarManager.h
--- a/Engine/Plugins/Runtime/Oculus/OculusAvatar/Source/Public/OvrAvatarManager.h
+++ b/Engine/Plugins/Runtime/Oculus/OculusAvatar/Source/Public/OvrAvatarManager.h
@@ -75,6 +75,7 @@ private:
 	TMap<FString, AvatarPacketQueue*> AvatarPacketQueues;
 	
 	void* OVRPluginHandle = nullptr;
+	void* OVRAvatarHandle = nullptr;
 
 	ovrAvatarLogLevel LogLevel = ovrAvatarLogLevel::ovrAvatarLogLevel_Silent;
 
