Don't use _t types

diff --git a/Engine/Plugins/Runtime/Oculus/OculusAvatar/Source/Public/OvrAvatar.h b/Engine/Plugins/Runtime/Oculus/OculusAvatar/Source/Public/OvrAvatar.h
--- a/Engine/Plugins/Runtime/Oculus/OculusAvatar/Source/Public/OvrAvatar.h
+++ b/Engine/Plugins/Runtime/Oculus/OculusAvatar/Source/Public/OvrAvatar.h
@@ -120,9 +120,9 @@ protected:
 
 	uint64_t OnlineUserID = 0;
 
-	TSet<ovrAvatarAssetID> AssetIds;
-	TMap<ovrAvatarAssetID, TWeakObjectPtr<UPoseableMeshComponent>> MeshComponents;
-	TMap<ovrAvatarAssetID, TWeakObjectPtr<UPoseableMeshComponent>> DepthMeshComponents;
+	TSet<uint64> AssetIds;
+	TMap<uint64, TWeakObjectPtr<UPoseableMeshComponent>> MeshComponents;
+	TMap<uint64, TWeakObjectPtr<UPoseableMeshComponent>> DepthMeshComponents;
 
 	ovrAvatar* Avatar = nullptr;
 
diff --git a/Engine/Plugins/Runtime/Oculus/OculusAvatar/Source/Public/OvrAvatarManager.h b/Engine/Plugins/Runtime/Oculus/OculusAvatar/Source/Public/OvrAvatarManager.h
--- a/Engine/Plugins/Runtime/Oculus/OculusAvatar/Source/Public/OvrAvatarManager.h
+++ b/Engine/Plugins/Runtime/Oculus/OculusAvatar/Source/Public/OvrAvatarManager.h
@@ -54,8 +54,8 @@ private:
 
 	bool IsInitialized = false;
 
-	TMap<uint64_t, TWeakObjectPtr<UTexture>> Textures;
-	TSet<uint64_t> NormalMapIDs;
+	TMap<uint64, TWeakObjectPtr<UTexture>> Textures;
+	TSet<uint64> NormalMapIDs;
 
 	static FOvrAvatarManager* sAvatarManager;
 	const char* AVATAR_APP_ID = nullptr;
