Cherrypicking CL#8129800 from Dev-VR:
Work around crash in UOvrAvatarGazeTarget::BeginDestroy()

diff --git a/Engine/Plugins/Runtime/Oculus/OculusAvatar/Source/Public/OculusAvatarModule.h b/Engine/Plugins/Runtime/Oculus/OculusAvatar/Source/Public/OculusAvatarModule.h
--- a/Engine/Plugins/Runtime/Oculus/OculusAvatar/Source/Public/OculusAvatarModule.h
+++ b/Engine/Plugins/Runtime/Oculus/OculusAvatar/Source/Public/OculusAvatarModule.h
@@ -3,11 +3,11 @@
 #pragma once
  
 #include "Modules/ModuleManager.h"
- 
+
 class OculusAvatarModule : public IModuleInterface
 {
 public:
 	/** IModuleInterface implementation */
-	void StartupModule();
-	void ShutdownModule();
+	virtual void StartupModule() override;
+	virtual void ShutdownModule() override;
 };
diff --git a/Engine/Plugins/Runtime/Oculus/OculusAvatar/Source/Public/OvrAvatarGazeTarget.cpp b/Engine/Plugins/Runtime/Oculus/OculusAvatar/Source/Public/OvrAvatarGazeTarget.cpp
--- a/Engine/Plugins/Runtime/Oculus/OculusAvatar/Source/Public/OvrAvatarGazeTarget.cpp
+++ b/Engine/Plugins/Runtime/Oculus/OculusAvatar/Source/Public/OvrAvatarGazeTarget.cpp
@@ -38,10 +38,12 @@ static const FString GazeTargetToString(ovrAvatarGazeTargetType targetType)
 
 UOvrAvatarGazeTarget::UOvrAvatarGazeTarget()
 : GazeTransform(nullptr)
+, bShuttingDown(false)
 {
 	PrimaryComponentTick.bCanEverTick = true;
+	UOvrAvatarManager::Get().OnShutdown().AddUObject(this, UOvrAvatarGazeTarget::HandleShutdownEvent);
 
-	IsEnabled = true;
+	bIsEnabled = true;
 
 	NativeTarget.targetCount = 1;
 	NativeTarget.targets[0].type = ovrAvatarGazeTargetType_Object;
@@ -49,9 +51,18 @@ UOvrAvatarGazeTarget::UOvrAvatarGazeTarget()
 	OvrAvatarHelpers::OvrAvatarVec3Zero(NativeTarget.targets[0].worldPosition);
 }
 
+void UOvrAvatarGazeTarget::HandleShutdownEvent()
+{
+	bShuttingDown = true;
+}
+
 void UOvrAvatarGazeTarget::BeginPlay() 
 { 
 	Super::BeginPlay(); 
+	if (bShuttingDown)
+	{
+		return;
+	}
 
 	NativeTarget.targets[0].type = ConvertEditorTypeToNativeType();
 	ovrAvatar_UpdateGazeTargets(&NativeTarget);
@@ -71,8 +82,10 @@ void UOvrAvatarGazeTarget::TickComponent(float DeltaTime, enum ELevelTick TickTy
 { 
 	Super::TickComponent(DeltaTime, TickType, ThisTickFunction); 
 
-	if (!IsEnabled)
+	if (!bIsEnabled || bShuttingDown)
+	{
 		return;
+	}
 
 	const FTransform& GazeTrans = GetGazeTransform();
 	const FVector worldPosition = GazeTrans.GetLocation();
@@ -105,12 +118,12 @@ const FTransform& UOvrAvatarGazeTarget::GetGazeTransform() const
 
 void UOvrAvatarGazeTarget::EnableGazeTarget(bool DoEnable)
 {
-	if (!DoEnable && IsEnabled)
+	if (!DoEnable && bIsEnabled && !bShuttingDown)
 	{
 		ovrAvatar_RemoveGazeTargets(1, &NativeTarget.targets[0].id);
 	}
 
-	IsEnabled = DoEnable;
+	bIsEnabled = DoEnable;
 }
 
 ovrAvatarGazeTargetType UOvrAvatarGazeTarget::ConvertEditorTypeToNativeType() const
diff --git a/Engine/Plugins/Runtime/Oculus/OculusAvatar/Source/Public/OvrAvatarGazeTarget.h b/Engine/Plugins/Runtime/Oculus/OculusAvatar/Source/Public/OvrAvatarGazeTarget.h
--- a/Engine/Plugins/Runtime/Oculus/OculusAvatar/Source/Public/OvrAvatarGazeTarget.h
+++ b/Engine/Plugins/Runtime/Oculus/OculusAvatar/Source/Public/OvrAvatarGazeTarget.h
@@ -40,11 +40,13 @@ public:
 	void SetGazeTargetType(OculusAvatarGazeTargetType type) { TargetType = type; }
 	void SetAvatarHeadTransform(const FTransform& trans) { AvatarHeadTransform = trans; }
 private:
+	void HandleShutdownEvent();
 	const FTransform& GetGazeTransform() const;
 	ovrAvatarGazeTargetType ConvertEditorTypeToNativeType() const;
 
 	TWeakObjectPtr<USceneComponent> GazeTransform = nullptr;
 	ovrAvatarGazeTargetSet NativeTarget;
 	FTransform AvatarHeadTransform = FTransform::Identity;
-	bool IsEnabled = false;
+	bool bIsEnabled = false;
+	bool bShuttingDown = false;
 };
diff --git a/Engine/Plugins/Runtime/Oculus/OculusAvatar/Source/Public/OvrAvatarManager.cpp b/Engine/Plugins/Runtime/Oculus/OculusAvatar/Source/Public/OvrAvatarManager.cpp
--- a/Engine/Plugins/Runtime/Oculus/OculusAvatar/Source/Public/OvrAvatarManager.cpp
+++ b/Engine/Plugins/Runtime/Oculus/OculusAvatar/Source/Public/OvrAvatarManager.cpp
@@ -167,8 +167,10 @@ static const FString TextureFormatToString(ovrAvatarTextureFormat format)
 
 void UOvrAvatarManager::Tick(float DeltaTime)
 {
-	if (!IsInitialized)
+	if (!bIsInitialized)
+	{
 		return;
+	}
 
 	while (ovrAvatarMessage* message = ovrAvatarMessage_Pop())
 	{
@@ -198,7 +200,7 @@ void UOvrAvatarManager::InitializeSDK()
 {
 	UE_LOG(LogAvatars, Display, TEXT("UOvrAvatarManager::InitializeSDK()"));
 
-	if (!IsInitialized)
+	if (!bIsInitialized)
 	{
 #if WITH_EDITORONLY_DATA
 		for (int32 AssetIndex = 0; AssetIndex < ARRAY_COUNT(AssetList); ++AssetIndex)
@@ -242,7 +244,7 @@ void UOvrAvatarManager::InitializeSDK()
 		ovrAvatar_Initialize(TCHAR_TO_ANSI(*AVATAR_APP_ID));
 #endif
 
-		IsInitialized = true;
+		bIsInitialized = true;
 
 		ovrAvatar_SetLoggingLevel(LogLevel);
 
@@ -258,18 +260,19 @@ void UOvrAvatarManager::InitializeSDK()
 
 void UOvrAvatarManager::ShutdownSDK()
 {
-	if (IsInitialized)
+	if (bIsInitialized)
 	{
 #if WITH_EDITORONLY_DATA
 		AssetObjects.Empty();
 #endif
+		ShutdownEvent.Broadcast();
 
 		// This is crashing when cooking content, and not sure why...
 #if !WITH_EDITORONLY_DATA
 		ovrAvatar_RegisterLoggingCallback(nullptr);
 #endif
 
-		IsInitialized = false;
+		bIsInitialized = false;
 		ovrAvatar_Shutdown();
 	}
 }
diff --git a/Engine/Plugins/Runtime/Oculus/OculusAvatar/Source/Public/OvrAvatarManager.h b/Engine/Plugins/Runtime/Oculus/OculusAvatar/Source/Public/OvrAvatarManager.h
--- a/Engine/Plugins/Runtime/Oculus/OculusAvatar/Source/Public/OvrAvatarManager.h
+++ b/Engine/Plugins/Runtime/Oculus/OculusAvatar/Source/Public/OvrAvatarManager.h
@@ -55,6 +55,10 @@ public:
 
 	void SetSDKLoggingLevel(ovrAvatarLogLevel level) { ovrAvatar_SetLoggingLevel(level); }
 
+	/** Event used to avoid accessing the avatar library after shutdown */
+	DECLARE_EVENT(UOvrAvatarManager, FAvatarShutdownEvent)
+	FAvatarShutdownEvent& OnShutdown() { return ShutdownEvent; }
+
 private:
 	static void SDKLogger(const char * str);
 
@@ -66,7 +70,7 @@ private:
 
 	UTexture2D* LoadTexture(const uint64_t id, const ovrAvatarTextureAssetData* data, bool isLinearColor);
 
-	bool IsInitialized = false;
+	bool bIsInitialized = false;
 
 	UPROPERTY()
 	TMap<uint64, UTexture*> Textures;
@@ -93,6 +97,7 @@ private:
 
 	void* OVRPluginHandle = nullptr;
 	void* OVRAvatarHandle = nullptr;
+	FAvatarShutdownEvent ShutdownEvent;
 
 	ovrAvatarLogLevel LogLevel = ovrAvatarLogLevel::ovrAvatarLogLevel_Silent;
 
