Fix build error caused by CL#8129845

diff --git a/Engine/Plugins/Runtime/Oculus/OculusAvatar/Source/Public/OvrAvatarGazeTarget.cpp b/Engine/Plugins/Runtime/Oculus/OculusAvatar/Source/Public/OvrAvatarGazeTarget.cpp
--- a/Engine/Plugins/Runtime/Oculus/OculusAvatar/Source/Public/OvrAvatarGazeTarget.cpp
+++ b/Engine/Plugins/Runtime/Oculus/OculusAvatar/Source/Public/OvrAvatarGazeTarget.cpp
@@ -41,7 +41,7 @@ UOvrAvatarGazeTarget::UOvrAvatarGazeTarget()
 , bShuttingDown(false)
 {
 	PrimaryComponentTick.bCanEverTick = true;
-	UOvrAvatarManager::Get().OnShutdown().AddUObject(this, UOvrAvatarGazeTarget::HandleShutdownEvent);
+	UOvrAvatarManager::Get().OnShutdown().AddUObject(this, &UOvrAvatarGazeTarget::HandleShutdownEvent);
 
 	bIsEnabled = true;
 
