Fixing missing includes..

diff --git a/Engine/Plugins/Runtime/Oculus/OculusAvatar/Source/Public/OvrAvatar.cpp b/Engine/Plugins/Runtime/Oculus/OculusAvatar/Source/Public/OvrAvatar.cpp
--- a/Engine/Plugins/Runtime/Oculus/OculusAvatar/Source/Public/OvrAvatar.cpp
+++ b/Engine/Plugins/Runtime/Oculus/OculusAvatar/Source/Public/OvrAvatar.cpp
@@ -8,7 +8,9 @@
 #include "Rendering/SkeletalMeshRenderData.h"
 #include "Rendering/SkeletalMeshLODModel.h"
 #include "Rendering/SkeletalMeshModel.h"
-
+#include "Materials/MaterialInstanceDynamic.h"
+#include "Engine/SkeletalMesh.h"
+#include "Animation/Skeleton.h"
 
 float DebugLineScale = 100.f;
 bool DrawDebug = false;
diff --git a/Engine/Plugins/Runtime/Oculus/OculusAvatar/Source/Public/OvrAvatar.h b/Engine/Plugins/Runtime/Oculus/OculusAvatar/Source/Public/OvrAvatar.h
--- a/Engine/Plugins/Runtime/Oculus/OculusAvatar/Source/Public/OvrAvatar.h
+++ b/Engine/Plugins/Runtime/Oculus/OculusAvatar/Source/Public/OvrAvatar.h
@@ -10,6 +10,8 @@
 #include "OvrAvatar.generated.h"
 
 class UPoseableMeshComponent;
+class USkeletalMesh;
+class UMeshComponent;
 
 UCLASS()
 class OCULUSAVATAR_API UOvrAvatar : public UActorComponent
diff --git a/Engine/Plugins/Runtime/Oculus/OculusAvatar/Source/Public/OvrAvatarManager.cpp b/Engine/Plugins/Runtime/Oculus/OculusAvatar/Source/Public/OvrAvatarManager.cpp
--- a/Engine/Plugins/Runtime/Oculus/OculusAvatar/Source/Public/OvrAvatarManager.cpp
+++ b/Engine/Plugins/Runtime/Oculus/OculusAvatar/Source/Public/OvrAvatarManager.cpp
@@ -2,10 +2,10 @@
 #include "OvrAvatarManager.h"
 #include "OvrAvatar.h"
 #include "OVR_Avatar.h"
-#include "Engine/Texture2D.h"
+#include "RenderUtils.h"
+#include "UObject/UObjectIterator.h"
 #include "OculusHMDModule.h"
 
-
 DEFINE_LOG_CATEGORY(LogAvatars);
 
 FOvrAvatarManager* FOvrAvatarManager::sAvatarManager = nullptr;
diff --git a/Engine/Plugins/Runtime/Oculus/OculusAvatar/Source/Public/OvrAvatarManager.h b/Engine/Plugins/Runtime/Oculus/OculusAvatar/Source/Public/OvrAvatarManager.h
--- a/Engine/Plugins/Runtime/Oculus/OculusAvatar/Source/Public/OvrAvatarManager.h
+++ b/Engine/Plugins/Runtime/Oculus/OculusAvatar/Source/Public/OvrAvatarManager.h
@@ -2,12 +2,11 @@
 #include "Containers/Ticker.h"
 #include "OVR_Avatar.h"
 #include "Containers/Map.h"
+#include "Containers/Queue.h"
+#include "Engine/Texture2D.h"
 
 DECLARE_LOG_CATEGORY_EXTERN(LogAvatars, Log, All);
 
-class UTexture2D;
-class UTexture;
-
 class OCULUSAVATAR_API FOvrAvatarManager : public FTickerObjectBase
 {
 public:
