[Oculus Avatars] - Fixing broken skeleton updates in 4.21

Summary:
- A dirty transform flag was introduced into the PoseableMeshComponent in 4.21.  This was never getting set but the avatar bone update flow.
- The mark dirty function is currently private, which makes it impossible to get slave components to update. Their ShouldUpdateTransform call in SkinnedMeshComponent will always be false.

Mirroring change submitted to 4.21 hotfix, but making MarkRefreshTransformDirty public instead of defining UOvrAvatar as friend.

diff --git a/Engine/Plugins/Runtime/Oculus/OculusAvatar/Source/Public/OvrAvatar.cpp b/Engine/Plugins/Runtime/Oculus/OculusAvatar/Source/Public/OvrAvatar.cpp
--- a/Engine/Plugins/Runtime/Oculus/OculusAvatar/Source/Public/OvrAvatar.cpp
+++ b/Engine/Plugins/Runtime/Oculus/OculusAvatar/Source/Public/OvrAvatar.cpp
@@ -468,6 +468,7 @@ void UOvrAvatar::UpdatePostSDK()
 					if (MeshVisible && IsSelfOccluding)
 					{
 						UpdateMeshComponent(*depthMesh, RenderData->localTransform);
+						depthMesh->MarkRefreshTransformDirty();
 					}
 
 					depthMesh->SetVisibility(MeshVisible && IsSelfOccluding, true);
@@ -530,6 +531,7 @@ void UOvrAvatar::UpdatePostSDK()
 					if (MeshVisible && IsSelfOccluding)
 					{
 						UpdateMeshComponent(*depthMesh, RenderData->localTransform);
+						depthMesh->MarkRefreshTransformDirty();
 					}
 
 					depthMesh->SetVisibility(MeshVisible && IsSelfOccluding, true);
@@ -629,6 +631,8 @@ void UOvrAvatar::UpdateSkeleton(UPoseableMeshComponent& mesh, const ovrAvatarSki
 		OvrAvatarHelpers::ConvertTransform(pose.jointTransform[BoneIndex], LocalBone);
 		mesh.BoneSpaceTransforms[BoneIndex] = LocalBone;
 	}
+
+	mesh.MarkRefreshTransformDirty();
 }
 
 USceneComponent* UOvrAvatar::DetachHand(HandType hand)
