Missing includes to fix CIS

diff --git a/Engine/Plugins/Runtime/Oculus/OculusAvatar/Source/Public/OvrAvatarGazeTarget.cpp b/Engine/Plugins/Runtime/Oculus/OculusAvatar/Source/Public/OvrAvatarGazeTarget.cpp
--- a/Engine/Plugins/Runtime/Oculus/OculusAvatar/Source/Public/OvrAvatarGazeTarget.cpp
+++ b/Engine/Plugins/Runtime/Oculus/OculusAvatar/Source/Public/OvrAvatarGazeTarget.cpp
@@ -1,6 +1,7 @@
 #include "OvrAvatarGazeTarget.h"
 #include "HAL/IConsoleManager.h"
 #include "OvrAvatarHelpers.h"
+#include "OvrAvatarManager.h"
 #include "GameFramework/Actor.h"
 
 static int32 GOvrAvatarGazeTargetDebugDraw = 0;
diff --git a/Engine/Plugins/Runtime/Oculus/OculusAvatar/Source/Public/OvrAvatarGazeTarget.h b/Engine/Plugins/Runtime/Oculus/OculusAvatar/Source/Public/OvrAvatarGazeTarget.h
--- a/Engine/Plugins/Runtime/Oculus/OculusAvatar/Source/Public/OvrAvatarGazeTarget.h
+++ b/Engine/Plugins/Runtime/Oculus/OculusAvatar/Source/Public/OvrAvatarGazeTarget.h
@@ -2,6 +2,7 @@
 
 #include "OVR_Avatar.h"
 #include "Components/ActorComponent.h"
+#include "Runtime/Engine/Classes/Components/SceneComponent.h"
 #include "OvrAvatarGazeTarget.generated.h"
 
 UENUM(BlueprintType, Category = OculusAvatar)
