Fix oculus code
#rnx

diff --git a/Engine/Plugins/Runtime/Oculus/OculusAvatar/Source/Public/OvrAvatar.cpp b/Engine/Plugins/Runtime/Oculus/OculusAvatar/Source/Public/OvrAvatar.cpp
--- a/Engine/Plugins/Runtime/Oculus/OculusAvatar/Source/Public/OvrAvatar.cpp
+++ b/Engine/Plugins/Runtime/Oculus/OculusAvatar/Source/Public/OvrAvatar.cpp
@@ -2379,9 +2379,9 @@ void UOvrAvatar::LoadCombinedMesh(USkeletalMesh* SkeletalMesh, const ovrAvatarMe
 	FSkeletalMeshLODInfo& LodInfo = SkeletalMesh->AddLODInfo();
 	LodInfo.ScreenSize = 0.3f;
 	LodInfo.LODHysteresis = 0.2f;
+	LodInfo.BuildSettings.bUseFullPrecisionUVs = true;
 
 	SkeletalMesh->RefSkeleton.Empty(data->skinnedBindPose.jointCount);
-	SkeletalMesh->bUseFullPrecisionUVs = true;
 	SkeletalMesh->bHasBeenSimplified = false;
 	SkeletalMesh->bHasVertexColors = true;
 
