Speculative fix for Ovr Avatar shutdown crash. Crash in question always occurs during shutdown, disabling the gaze target, and it always happens in the delay load helper functionality triggering an exception thrown if the module has been unloaded. Adding this check should catch this shut-down ordering case.

diff --git a/Engine/Plugins/Runtime/Oculus/OculusAvatar/Source/Public/OvrAvatarGazeTarget.cpp b/Engine/Plugins/Runtime/Oculus/OculusAvatar/Source/Public/OvrAvatarGazeTarget.cpp
--- a/Engine/Plugins/Runtime/Oculus/OculusAvatar/Source/Public/OvrAvatarGazeTarget.cpp
+++ b/Engine/Plugins/Runtime/Oculus/OculusAvatar/Source/Public/OvrAvatarGazeTarget.cpp
@@ -118,7 +118,7 @@ const FTransform& UOvrAvatarGazeTarget::GetGazeTransform() const
 
 void UOvrAvatarGazeTarget::EnableGazeTarget(bool DoEnable)
 {
-	if (!DoEnable && bIsEnabled && !bShuttingDown)
+	if (!DoEnable && bIsEnabled && !bShuttingDown && UOvrAvatarManager::Get().IsOVRPluginValid())
 	{
 		ovrAvatar_RemoveGazeTargets(1, &NativeTarget.targets[0].id);
 	}
