Fix out-of-order destruction crashes with Oculus Avatar plugin.

diff --git a/Engine/Plugins/Runtime/Oculus/OculusAvatar/Source/Public/OculusAvatarModule.cpp b/Engine/Plugins/Runtime/Oculus/OculusAvatar/Source/Public/OculusAvatarModule.cpp
--- a/Engine/Plugins/Runtime/Oculus/OculusAvatar/Source/Public/OculusAvatarModule.cpp
+++ b/Engine/Plugins/Runtime/Oculus/OculusAvatar/Source/Public/OculusAvatarModule.cpp
@@ -11,7 +11,6 @@ void OculusAvatarModule::StartupModule()
  
 void OculusAvatarModule::ShutdownModule()
 {
-	UOvrAvatarManager::Get().ShutdownSDK();
 	UOvrAvatarManager::Destroy();
 }
 
diff --git a/Engine/Plugins/Runtime/Oculus/OculusAvatar/Source/Public/OvrAvatarManager.cpp b/Engine/Plugins/Runtime/Oculus/OculusAvatar/Source/Public/OvrAvatarManager.cpp
--- a/Engine/Plugins/Runtime/Oculus/OculusAvatar/Source/Public/OvrAvatarManager.cpp
+++ b/Engine/Plugins/Runtime/Oculus/OculusAvatar/Source/Public/OvrAvatarManager.cpp
@@ -15,6 +15,7 @@
 DEFINE_LOG_CATEGORY(LogAvatars);
 
 UOvrAvatarManager* UOvrAvatarManager::sAvatarManager = nullptr;
+void* UOvrAvatarManager::OVRAvatarHandle = nullptr;
 
 FSoftObjectPath UOvrAvatarManager::AssetList[] =
 {
@@ -129,13 +130,12 @@ void UOvrAvatarManager::Destroy()
 {
 	if (sAvatarManager)
 	{
-		sAvatarManager->RemoveFromRoot();
+		sAvatarManager->ShutdownSDK();
+		UOvrAvatarManager* temp_avatar_manager = sAvatarManager;
 		sAvatarManager = nullptr;
+		temp_avatar_manager->RemoveFromRoot();
 	}
-}
 
-UOvrAvatarManager::~UOvrAvatarManager()
-{
 	if (OVRAvatarHandle)
 	{
 		FPlatformProcess::FreeDllHandle(OVRAvatarHandle);
@@ -143,6 +143,17 @@ UOvrAvatarManager::~UOvrAvatarManager()
 	}
 }
 
+UOvrAvatarManager::~UOvrAvatarManager()
+{
+	if (sAvatarManager != nullptr)
+	{
+		UE_LOG(LogAvatars, Log, TEXT("Shutdown ordering error closing down OVRAvatar module."));
+		sAvatarManager = nullptr;
+
+		ShutdownSDK();
+	}
+}
+
 static const FString sTextureFormatStrings[ovrAvatarTextureFormat_Count] =
 {
 	FString("ovrAvatarTextureFormat_RGB24"),
diff --git a/Engine/Plugins/Runtime/Oculus/OculusAvatar/Source/Public/OvrAvatarManager.h b/Engine/Plugins/Runtime/Oculus/OculusAvatar/Source/Public/OvrAvatarManager.h
--- a/Engine/Plugins/Runtime/Oculus/OculusAvatar/Source/Public/OvrAvatarManager.h
+++ b/Engine/Plugins/Runtime/Oculus/OculusAvatar/Source/Public/OvrAvatarManager.h
@@ -95,8 +95,9 @@ private:
 
 	TMap<FString, AvatarPacketQueue*> AvatarPacketQueues;
 
+	static void* OVRAvatarHandle;
+
 	void* OVRPluginHandle = nullptr;
-	void* OVRAvatarHandle = nullptr;
 	FAvatarShutdownEvent ShutdownEvent;
 
 	ovrAvatarLogLevel LogLevel = ovrAvatarLogLevel::ovrAvatarLogLevel_Silent;
