Verify we're not in shutdown before calling ovr DLL function

diff --git a/Engine/Plugins/Runtime/Oculus/OculusAvatar/Source/Public/OvrAvatarGazeTarget.cpp b/Engine/Plugins/Runtime/Oculus/OculusAvatar/Source/Public/OvrAvatarGazeTarget.cpp
--- a/Engine/Plugins/Runtime/Oculus/OculusAvatar/Source/Public/OvrAvatarGazeTarget.cpp
+++ b/Engine/Plugins/Runtime/Oculus/OculusAvatar/Source/Public/OvrAvatarGazeTarget.cpp
@@ -4,6 +4,7 @@
 #include "OvrAvatarHelpers.h"
 #include "OvrAvatarManager.h"
 #include "GameFramework/Actor.h"
+#include "CoreGlobals.h"
 
 static int32 GOvrAvatarGazeTargetDebugDraw = 0;
 static FAutoConsoleVariableRef CVarOvrAvatarGazeTargetDebugDraw(
@@ -118,7 +119,7 @@ const FTransform& UOvrAvatarGazeTarget::GetGazeTransform() const
 
 void UOvrAvatarGazeTarget::EnableGazeTarget(bool DoEnable)
 {
-	if (!DoEnable && bIsEnabled && !bShuttingDown && UOvrAvatarManager::Get().IsOVRPluginValid())
+	if (!IsEngineExitRequested() && !DoEnable && bIsEnabled && !bShuttingDown && UOvrAvatarManager::Get().IsOVRPluginValid())
 	{
 		ovrAvatar_RemoveGazeTargets(1, &NativeTarget.targets[0].id);
 	}
