Rounding of FLinearColor to FColor (and [0,255] uint8 to [0,1] float) according to standard methods

diff --git a/Engine/Plugins/Enterprise/LidarPointCloud/Source/LidarPointCloudRuntime/Public/LidarPointCloudShared.h b/Engine/Plugins/Enterprise/LidarPointCloud/Source/LidarPointCloudRuntime/Public/LidarPointCloudShared.h
--- a/Engine/Plugins/Enterprise/LidarPointCloud/Source/LidarPointCloudRuntime/Public/LidarPointCloudShared.h
+++ b/Engine/Plugins/Enterprise/LidarPointCloud/Source/LidarPointCloudRuntime/Public/LidarPointCloudShared.h
@@ -378,7 +378,7 @@ public:
 	FLidarPointCloudPoint(const float& X, const float& Y, const float& Z, const float& I)
 		: FLidarPointCloudPoint(X, Y, Z)
 	{
-		Color.A = FMath::FloorToInt(FMath::Clamp(I, 0.0f, 1.0f) * 255.999f);
+		Color.A = FMath::RoundToInt(FMath::Clamp(I, 0.0f, 1.0f) * 255.f);
 	}
 	FLidarPointCloudPoint(const float& X, const float& Y, const float& Z, const float& R, const float& G, const float& B, const float& A = 1.0f)
 		: FLidarPointCloudPoint(X, Y, Z)
diff --git a/Engine/Plugins/Runtime/Oculus/OculusAvatar/Source/Public/OvrAvatar.cpp b/Engine/Plugins/Runtime/Oculus/OculusAvatar/Source/Public/OvrAvatar.cpp
--- a/Engine/Plugins/Runtime/Oculus/OculusAvatar/Source/Public/OvrAvatar.cpp
+++ b/Engine/Plugins/Runtime/Oculus/OculusAvatar/Source/Public/OvrAvatar.cpp
@@ -1557,7 +1557,7 @@ void UOvrAvatar::LoadMesh(USkeletalMesh* SkeletalMesh, const ovrAvatarMeshAssetD
 
 			for (uint32_t BlendIndex = 0; BlendIndex < MAX_TOTAL_INFLUENCES; BlendIndex++)
 			{
-				ModelVertex.InfluenceWeights[BlendIndex] = BlendIndex < NumBlendWeights ? (uint8_t)(255.9999f*SourceVertex->blendWeights[BlendIndex]) : 0;
+				ModelVertex.InfluenceWeights[BlendIndex] = BlendIndex < NumBlendWeights ? (uint8_t)(255.f*SourceVertex->blendWeights[BlendIndex]+0.5f) : 0;
 				ModelVertex.InfluenceBones[BlendIndex] = BlendIndex < NumBlendWeights ? SourceVertex->blendIndices[BlendIndex] : 0;
 
 				uint32 Weight = ModelVertex.InfluenceWeights[BlendIndex];
@@ -1603,7 +1603,7 @@ void UOvrAvatar::LoadMesh(USkeletalMesh* SkeletalMesh, const ovrAvatarMeshAssetD
 
 			for (uint32_t BlendIndex = 0; BlendIndex < MAX_TOTAL_INFLUENCES; BlendIndex++)
 			{
-				InWeights[VertIndex].InfluenceWeights[BlendIndex] = BlendIndex < NumBlendWeights ? (uint8_t)(255.9999f*SourceVertex->blendWeights[BlendIndex]) : 0;
+				InWeights[VertIndex].InfluenceWeights[BlendIndex] = BlendIndex < NumBlendWeights ? (uint8_t)(255.f*SourceVertex->blendWeights[BlendIndex]+0.5f) : 0;
 				InWeights[VertIndex].InfluenceBones[BlendIndex] = BlendIndex < NumBlendWeights ? SourceVertex->blendIndices[BlendIndex] : 0;
 
 				uint32 Weight = InWeights[VertIndex].InfluenceWeights[BlendIndex];
@@ -2323,7 +2323,7 @@ void UOvrAvatar::LoadCombinedMesh(USkeletalMesh* SkeletalMesh, const ovrAvatarMe
 
 			for (uint32_t BlendIndex = 0; BlendIndex < MAX_TOTAL_INFLUENCES; BlendIndex++)
 			{
-				DestVertex.InfluenceWeights[BlendIndex] = BlendIndex < NumBlendWeights ? (uint8_t)(255.9999f*SourceVertex->blendWeights[BlendIndex]) : 0;
+				DestVertex.InfluenceWeights[BlendIndex] = BlendIndex < NumBlendWeights ? (uint8_t)(255.f*SourceVertex->blendWeights[BlendIndex]+0.5f) : 0;
 				DestVertex.InfluenceBones[BlendIndex] = BlendIndex < NumBlendWeights ? SourceVertex->blendIndices[BlendIndex] : 0;
 
 				uint32 Weight = DestVertex.InfluenceWeights[BlendIndex];
@@ -2596,7 +2596,7 @@ void UOvrAvatar::LoadCombinedMesh(USkeletalMesh* SkeletalMesh, const ovrAvatarMe
 
 			for (uint32_t BlendIndex = 0; BlendIndex < MAX_TOTAL_INFLUENCES; BlendIndex++)
 			{
-				InWeights[CurrentVertexBufferCount].InfluenceWeights[BlendIndex] = BlendIndex < NumBlendWeights ? (uint8_t)(255.9999f*SourceVertex->blendWeights[BlendIndex]) : 0;
+				InWeights[CurrentVertexBufferCount].InfluenceWeights[BlendIndex] = BlendIndex < NumBlendWeights ? (uint8_t)(255.f*SourceVertex->blendWeights[BlendIndex]+0.5f) : 0;
 				InWeights[CurrentVertexBufferCount].InfluenceBones[BlendIndex] = BlendIndex < NumBlendWeights ? SourceVertex->blendIndices[BlendIndex] : 0;
 
 				uint32 Weight = InWeights[CurrentVertexBufferCount].InfluenceWeights[BlendIndex];
diff --git a/Engine/Source/Developer/TextureFormatIntelISPCTexComp/Private/TextureFormatIntelISPCTexComp.cpp b/Engine/Source/Developer/TextureFormatIntelISPCTexComp/Private/TextureFormatIntelISPCTexComp.cpp
--- a/Engine/Source/Developer/TextureFormatIntelISPCTexComp/Private/TextureFormatIntelISPCTexComp.cpp
+++ b/Engine/Source/Developer/TextureFormatIntelISPCTexComp/Private/TextureFormatIntelISPCTexComp.cpp
@@ -483,9 +483,9 @@ static void IntelASTCCompressScans(FASTCEncoderSettings* pEncSettings, FImage* p
 				FVector Normal = FVector(pInTexelsSwap[2] / 255.0f * 2.0f - 1.0f, pInTexelsSwap[1] / 255.0f * 2.0f - 1.0f, pInTexelsSwap[0] / 255.0f * 2.0f - 1.0f);
 				Normal = Normal.GetSafeNormal();
 				pInTexelsSwap[0] = 0;
-				pInTexelsSwap[1] = FMath::FloorToInt((Normal.Y * 0.5f + 0.5f) * 255.999f);
+				pInTexelsSwap[1] = FMath::RoundToInt((Normal.Y * 0.5f + 0.5f) * 255.f);
 				pInTexelsSwap[2] = 0;
-				pInTexelsSwap[3] = FMath::FloorToInt((Normal.X * 0.5f + 0.5f) * 255.999f);
+				pInTexelsSwap[3] = FMath::RoundToInt((Normal.X * 0.5f + 0.5f) * 255.f);
 
 				pInTexelsSwap += 4;
 			}
@@ -503,8 +503,8 @@ static void IntelASTCCompressScans(FASTCEncoderSettings* pEncSettings, FImage* p
 			{
 				FVector Normal = FVector(pInTexelsSwap[2] / 255.0f * 2.0f - 1.0f, pInTexelsSwap[1] / 255.0f * 2.0f - 1.0f, pInTexelsSwap[0] / 255.0f * 2.0f - 1.0f);
 				Normal = Normal.GetSafeNormal();
-				pInTexelsSwap[0] = FMath::FloorToInt((Normal.X * 0.5f + 0.5f) * 255.999f);
-				pInTexelsSwap[1] = FMath::FloorToInt((Normal.Y * 0.5f + 0.5f) * 255.999f);
+				pInTexelsSwap[0] = FMath::RoundToInt((Normal.X * 0.5f + 0.5f) * 255.f);
+				pInTexelsSwap[1] = FMath::RoundToInt((Normal.Y * 0.5f + 0.5f) * 255.f);
 				pInTexelsSwap[2] = 0;
 				pInTexelsSwap[3] = 255;
 
diff --git a/Engine/Source/Editor/DistCurveEditor/Private/CurveEditorViewportClient.cpp b/Engine/Source/Editor/DistCurveEditor/Private/CurveEditorViewportClient.cpp
--- a/Engine/Source/Editor/DistCurveEditor/Private/CurveEditorViewportClient.cpp
+++ b/Engine/Source/Editor/DistCurveEditor/Private/CurveEditorViewportClient.cpp
@@ -1409,20 +1409,20 @@ FColor FCurveEditorViewportClient::GetLineColor(FCurveEdInterface* EdInterface,
 			float Value;
 
 			Value = EdInterface->EvalSub(0, InVal);
-			Value *= 255.9f;
-			StepColor.R = FMath::TruncToInt(FMath::Clamp<float>(Value, 0.f, 255.9f));
+			Value *= 255.f;
+			StepColor.R = FMath::RoundToInt(FMath::Clamp<float>(Value, 0.f, 255.f));
 			Value = EdInterface->EvalSub(1, InVal);
-			Value *= 255.9f;
-			StepColor.G = FMath::TruncToInt(FMath::Clamp<float>(Value, 0.f, 255.9f));
+			Value *= 255.f;
+			StepColor.G = FMath::RoundToInt(FMath::Clamp<float>(Value, 0.f, 255.f));
 			Value = EdInterface->EvalSub(2, InVal);
-			Value *= 255.9f;
-			StepColor.B = FMath::TruncToInt(FMath::Clamp<float>(Value, 0.f, 255.9f));
+			Value *= 255.f;
+			StepColor.B = FMath::RoundToInt(FMath::Clamp<float>(Value, 0.f, 255.f));
 		}
 		else
 		{
-			StepColor.R = FMath::TruncToInt(FMath::Clamp<float>(EdInterface->EvalSub(0, InVal), 0.f, 255.9f));
-			StepColor.G = FMath::TruncToInt(FMath::Clamp<float>(EdInterface->EvalSub(1, InVal), 0.f, 255.9f));
-			StepColor.B = FMath::TruncToInt(FMath::Clamp<float>(EdInterface->EvalSub(2, InVal), 0.f, 255.9f));
+			StepColor.R = FMath::RoundToInt(FMath::Clamp<float>(EdInterface->EvalSub(0, InVal), 0.f, 255.f));
+			StepColor.G = FMath::RoundToInt(FMath::Clamp<float>(EdInterface->EvalSub(1, InVal), 0.f, 255.f));
+			StepColor.B = FMath::RoundToInt(FMath::Clamp<float>(EdInterface->EvalSub(2, InVal), 0.f, 255.f));
 		}
 		StepColor.A = 255;
 	}
@@ -1433,12 +1433,12 @@ FColor FCurveEditorViewportClient::GetLineColor(FCurveEdInterface* EdInterface,
 			float Value;
 
 			Value = EdInterface->EvalSub(0, InVal);
-			Value *= 255.9f;
-			StepColor.R = FMath::TruncToInt(FMath::Clamp<float>(Value, 0.f, 255.9f));
+			Value *= 255.f;
+			StepColor.R = FMath::RoundToInt(FMath::Clamp<float>(Value, 0.f, 255.f));
 		}
 		else
 		{
-			StepColor.R = FMath::TruncToInt(FMath::Clamp<float>(EdInterface->EvalSub(0, InVal), 0.f, 255.9f));
+			StepColor.R = FMath::RoundToInt(FMath::Clamp<float>(EdInterface->EvalSub(0, InVal), 0.f, 255.f));
 		}
 		StepColor.G = StepColor.R;
 		StepColor.B = StepColor.R;
diff --git a/Engine/Source/Editor/DistCurveEditor/Private/SDistributionCurveEditor.cpp b/Engine/Source/Editor/DistCurveEditor/Private/SDistributionCurveEditor.cpp
--- a/Engine/Source/Editor/DistCurveEditor/Private/SDistributionCurveEditor.cpp
+++ b/Engine/Source/Editor/DistCurveEditor/Private/SDistributionCurveEditor.cpp
@@ -608,18 +608,18 @@ void SDistributionCurveEditor::OnSetColor()
 	{
 		float Value;
 
-		Value	= EdInterface->GetKeyOut(0, SelKey.KeyIndex) * 255.9f;
-		InputColor.R = FMath::TruncToInt(Value);
-		Value	= EdInterface->GetKeyOut(1, SelKey.KeyIndex) * 255.9f;
-		InputColor.G = FMath::TruncToInt(Value);
-		Value	= EdInterface->GetKeyOut(2, SelKey.KeyIndex) * 255.9f;
-		InputColor.B = FMath::TruncToInt(Value);
+		Value	= EdInterface->GetKeyOut(0, SelKey.KeyIndex) * 255.f;
+		InputColor.R = FMath::RoundToInt(Value);
+		Value	= EdInterface->GetKeyOut(1, SelKey.KeyIndex) * 255.f;
+		InputColor.G = FMath::RoundToInt(Value);
+		Value	= EdInterface->GetKeyOut(2, SelKey.KeyIndex) * 255.f;
+		InputColor.B = FMath::RoundToInt(Value);
 	}
 	else
 	{
-		InputColor.R = FMath::TruncToInt(FMath::Clamp<float>(EdInterface->GetKeyOut(0, SelKey.KeyIndex), 0.f, 255.9f));
-		InputColor.G = FMath::TruncToInt(FMath::Clamp<float>(EdInterface->GetKeyOut(1, SelKey.KeyIndex), 0.f, 255.9f));
-		InputColor.B = FMath::TruncToInt(FMath::Clamp<float>(EdInterface->GetKeyOut(2, SelKey.KeyIndex), 0.f, 255.9f));
+		InputColor.R = FMath::RoundToInt(FMath::Clamp<float>(EdInterface->GetKeyOut(0, SelKey.KeyIndex), 0.f, 255.f));
+		InputColor.G = FMath::RoundToInt(FMath::Clamp<float>(EdInterface->GetKeyOut(1, SelKey.KeyIndex), 0.f, 255.f));
+		InputColor.B = FMath::RoundToInt(FMath::Clamp<float>(EdInterface->GetKeyOut(2, SelKey.KeyIndex), 0.f, 255.f));
 	}
 
 	//since the data isn't stored in standard colors, a temp color is used
@@ -638,19 +638,19 @@ void SDistributionCurveEditor::OnSetColor()
 		float Value;
 		if (Entry.bFloatingPointColorCurve)
 		{
-			Value	= (float)TempColor.R / 255.9f;
+			Value	= (float)TempColor.R / 255.f;
 			if (Entry.bClamp)
 			{
 				Value = FMath::Clamp<float>(Value, Entry.ClampLow, Entry.ClampHigh);
 			}
 			EdInterface->SetKeyOut(0, SelKey.KeyIndex, Value);
-			Value	= (float)TempColor.G / 255.9f;
+			Value	= (float)TempColor.G / 255.f;
 			if (Entry.bClamp)
 			{
 				Value = FMath::Clamp<float>(Value, Entry.ClampLow, Entry.ClampHigh);
 			}
 			EdInterface->SetKeyOut(1, SelKey.KeyIndex, Value);
-			Value	= (float)TempColor.B / 255.9f;
+			Value	= (float)TempColor.B / 255.f;
 			if (Entry.bClamp)
 			{
 				Value = FMath::Clamp<float>(Value, Entry.ClampLow, Entry.ClampHigh);
diff --git a/Engine/Source/Editor/UnrealEd/Private/Lightmass/Lightmass.cpp b/Engine/Source/Editor/UnrealEd/Private/Lightmass/Lightmass.cpp
--- a/Engine/Source/Editor/UnrealEd/Private/Lightmass/Lightmass.cpp
+++ b/Engine/Source/Editor/UnrealEd/Private/Lightmass/Lightmass.cpp
@@ -260,7 +260,7 @@ void CopyLightProfile( const ULightComponent* In, Lightmass::FLightData& Out, TA
 
 					FFloat16 HalfValue = *(FFloat16*)&MipData[ (SourceY * (uint32)Source.GetSizeX() * 8) + (SourceX * 8) ];
 					float Value = HalfValue;
-					OutLightProfileTextureData[y * Width + x] = (uint8)(Value * 255.9999f);
+					OutLightProfileTextureData[y * Width + x] = (uint8)(Value * 255.f + 0.5f);
 				}
 			}
 		}
diff --git a/Engine/Source/Runtime/Core/Public/Math/UnrealMathUtility.h b/Engine/Source/Runtime/Core/Public/Math/UnrealMathUtility.h
--- a/Engine/Source/Runtime/Core/Public/Math/UnrealMathUtility.h
+++ b/Engine/Source/Runtime/Core/Public/Math/UnrealMathUtility.h
@@ -1691,7 +1691,7 @@ public:
 	static uint8 Quantize8UnsignedByte(float x)
 	{
 		// 0..1 -> 0..255
-		int32 Ret = (int32)(x * 255.999f);
+		int32 Ret = (int32)(x * 255.f + 0.5f);
 
 		check(Ret >= 0);
 		check(Ret <= 255);
diff --git a/Engine/Source/Runtime/Engine/Private/Curves/CurveLinearColor.cpp b/Engine/Source/Runtime/Engine/Private/Curves/CurveLinearColor.cpp
--- a/Engine/Source/Runtime/Engine/Private/Curves/CurveLinearColor.cpp
+++ b/Engine/Source/Runtime/Engine/Private/Curves/CurveLinearColor.cpp
@@ -336,8 +336,8 @@ void UCurveLinearColor::Serialize(FArchive& Ar)
 
 void UCurveLinearColor::WritePixel(uint8* Pixel, const FLinearColor& Color)
 {
-	Pixel[0] = FMath::FloorToInt(Color.B * 255.999f);
-	Pixel[1] = FMath::FloorToInt(Color.G * 255.999f);
-	Pixel[2] = FMath::FloorToInt(Color.R * 255.999f);
-	Pixel[3] = FMath::FloorToInt(Color.A * 255.999f);
+	Pixel[0] = FMath::RoundToInt(Color.B * 255.f);
+	Pixel[1] = FMath::RoundToInt(Color.G * 255.f);
+	Pixel[2] = FMath::RoundToInt(Color.R * 255.f);
+	Pixel[3] = FMath::RoundToInt(Color.A * 255.f);
 }
diff --git a/Engine/Source/Runtime/Renderer/Private/SystemTextures.cpp b/Engine/Source/Runtime/Renderer/Private/SystemTextures.cpp
--- a/Engine/Source/Runtime/Renderer/Private/SystemTextures.cpp
+++ b/Engine/Source/Runtime/Renderer/Private/SystemTextures.cpp
@@ -387,8 +387,8 @@ void FSystemTextures::InitializeFeatureLevelDependentTextures(FRHICommandListImm
 					check(Desc.Format == PF_R8G8);
 
 					uint8* Dest = (uint8*)(DestBuffer + x * 2 + y * DestStride);
-					Dest[0] = (int32)(FMath::Clamp(A, 0.0f, 1.0f) * 255.9999f);
-					Dest[1] = (int32)(FMath::Clamp(B, 0.0f, 1.0f) * 255.9999f);
+					Dest[0] = (int32)(FMath::Clamp(A, 0.0f, 1.0f) * 255.f + 0.5f);
+					Dest[1] = (int32)(FMath::Clamp(B, 0.0f, 1.0f) * 255.f + 0.5f);
 				}
 			}
 		}
